Removed StaticMeshEditorExtension plugin from template to fix build breakage

diff --git a/Templates/TP_AEC_CollabBP/TP_AEC_CollabBP.uproject b/Templates/TP_AEC_CollabBP/TP_AEC_CollabBP.uproject
--- a/Templates/TP_AEC_CollabBP/TP_AEC_CollabBP.uproject
+++ b/Templates/TP_AEC_CollabBP/TP_AEC_CollabBP.uproject
@@ -46,10 +46,6 @@
 				"Win64"
 			]
 		},
-		{
-			"Name": "StaticMeshEditorExtension",
-			"Enabled": true
-		},
 		{
 			"Name": "DatasmithGLTFImporter",
 			"Enabled": true
