Fixe warning UE5 from CollabViewer --  SampleLevel - Failed to find soft reference

diff --git a/Templates/TP_AEC_CollabBP/TP_AEC_CollabBP.uproject b/Templates/TP_AEC_CollabBP/TP_AEC_CollabBP.uproject
--- a/Templates/TP_AEC_CollabBP/TP_AEC_CollabBP.uproject
+++ b/Templates/TP_AEC_CollabBP/TP_AEC_CollabBP.uproject
@@ -80,8 +80,10 @@
 			"Name": "OpenXR",
 			"Enabled": true,
 			"SupportedTargetPlatforms": [
-				"Win32",
-				"Win64"
+				"Win64",
+				"Linux",
+				"Android",
+				"HoloLens"
 			]
 		},
 		{
