Collab Viewer template uproject file updated for openxr platform change.

diff --git a/Templates/TP_AEC_CollabBP/TP_AEC_CollabBP.uproject b/Templates/TP_AEC_CollabBP/TP_AEC_CollabBP.uproject
--- a/Templates/TP_AEC_CollabBP/TP_AEC_CollabBP.uproject
+++ b/Templates/TP_AEC_CollabBP/TP_AEC_CollabBP.uproject
@@ -61,13 +61,7 @@
 		},
 		{
 			"Name": "OpenXR",
-			"Enabled": true,
-			"SupportedTargetPlatforms": [
-				"Win64",
-				"Linux",
-				"Android",
-				"HoloLens"
-			]
+			"Enabled": true
 		}
 	]
 }
\ No newline at end of file
